- Add missing include to tag widget

diff --git a/Engine/Source/Editor/GameplayTagsEditor/Private/SGameplayTagWidget.cpp b/Engine/Source/Editor/GameplayTagsEditor/Private/SGameplayTagWidget.cpp
--- a/Engine/Source/Editor/GameplayTagsEditor/Private/SGameplayTagWidget.cpp
+++ b/Engine/Source/Editor/GameplayTagsEditor/Private/SGameplayTagWidget.cpp
@@ -5,6 +5,7 @@
 #include "GameplayTagsModule.h"
 #include "GameplayTags.h"
 #include "ScopedTransaction.h"
+#include "PropertyHandle.h"
 
 #define LOCTEXT_NAMESPACE "GameplayTagWidget"
 
