- Actually include correct header for tag widget. Mistakenly built development w/o editor the first time.

diff --git a/Engine/Source/Editor/GameplayTagsEditor/Private/SGameplayTagWidget.cpp b/Engine/Source/Editor/GameplayTagsEditor/Private/SGameplayTagWidget.cpp
--- a/Engine/Source/Editor/GameplayTagsEditor/Private/SGameplayTagWidget.cpp
+++ b/Engine/Source/Editor/GameplayTagsEditor/Private/SGameplayTagWidget.cpp
@@ -5,7 +5,7 @@
 #include "GameplayTagsModule.h"
 #include "GameplayTags.h"
 #include "ScopedTransaction.h"
-#include "PropertyHandle.h"
+#include "Editor/PropertyEditor/Public/PropertyHandle.h"
 
 #define LOCTEXT_NAMESPACE "GameplayTagWidget"
 
