GameplayTags - allow UFunctions with "GameplayTagFilter" metadata to filter gameplay tags of their parameters. (All parameters will have the same filtering rules applied)

diff --git a/Engine/Source/Editor/GameplayTagsEditor/Private/SGameplayTagGraphPin.cpp b/Engine/Source/Editor/GameplayTagsEditor/Private/SGameplayTagGraphPin.cpp
--- a/Engine/Source/Editor/GameplayTagsEditor/Private/SGameplayTagGraphPin.cpp
+++ b/Engine/Source/Editor/GameplayTagsEditor/Private/SGameplayTagGraphPin.cpp
@@ -44,6 +44,21 @@ void SGameplayTagGraphPin::ParseDefaultValueData()
 {
 	FString TagString = GraphPinObj->GetDefaultAsString();
 
+	UK2Node_CallFunction* CallFuncNode = Cast<UK2Node_CallFunction>(GraphPinObj->GetOuter());
+	
+	FilterString.Empty();
+	if (CallFuncNode)
+	{
+		UFunction* ThisFunction = CallFuncNode->GetTargetFunction();
+		if (ThisFunction)
+		{
+			if (ThisFunction->HasMetaData(TEXT("GameplayTagFilter")))
+			{
+				FilterString = ThisFunction->GetMetaData(TEXT("GameplayTagFilter"));
+			}
+		}
+	}
+
 	if (TagString.StartsWith(TEXT("(")) && TagString.EndsWith(TEXT(")")))
 	{
 		TagString = TagString.LeftChop(1);
@@ -78,6 +93,7 @@ TSharedRef<SWidget> SGameplayTagGraphPin::GetListContent()
 			.TagContainerName( TEXT("SGameplayTagGraphPin") )
 			.Visibility( this, &SGraphPin::GetDefaultValueVisibility )
 			.MultiSelect(false)
+			.Filter(FilterString)
 		];
 }
 
diff --git a/Engine/Source/Editor/GameplayTagsEditor/Private/SGameplayTagGraphPin.h b/Engine/Source/Editor/GameplayTagsEditor/Private/SGameplayTagGraphPin.h
--- a/Engine/Source/Editor/GameplayTagsEditor/Private/SGameplayTagGraphPin.h
+++ b/Engine/Source/Editor/GameplayTagsEditor/Private/SGameplayTagGraphPin.h
@@ -58,4 +58,6 @@ private:
 	// The List View used to display the read only tag names on the node.
 	TSharedPtr<SListView<TSharedPtr<FString>>> TagListView;
 
+	FString FilterString;
+
 };
