make asset registry dependency explicit for GameplayTagsEditor

diff --git a/Engine/Source/Editor/GameplayTagsEditor/GameplayTagsEditor.Build.cs b/Engine/Source/Editor/GameplayTagsEditor/GameplayTagsEditor.Build.cs
--- a/Engine/Source/Editor/GameplayTagsEditor/GameplayTagsEditor.Build.cs
+++ b/Engine/Source/Editor/GameplayTagsEditor/GameplayTagsEditor.Build.cs
@@ -9,9 +9,17 @@ namespace UnrealBuildTool.Rules
 			PrivateIncludePaths.AddRange(
 				new string[] {
 					"Editor/GameplayTagsEditor/Private",
+                    "Developer/AssetTools/Private",
 				}
 			);
 
+            PublicIncludePathModuleNames.AddRange(
+                new string[] {
+                    "AssetTools",
+                    "AssetRegistry",
+			    }
+            );
+
 
 			PrivateDependencyModuleNames.AddRange(
 				new string[]
@@ -21,6 +29,7 @@ namespace UnrealBuildTool.Rules
 					"CoreUObject",
 					"Engine",
 					"AssetTools",
+                    "AssetRegistry",
 					"GameplayTags",
                     "InputCore",
 					"Slate",
