Gameplay Tag custom editor works in structure editor.

UE-12154 Crash adding a Gameplay Tag variable to a struct

diff --git a/Engine/Source/Editor/GameplayTagsEditor/Private/GameplayTagCustomization.cpp b/Engine/Source/Editor/GameplayTagsEditor/Private/GameplayTagCustomization.cpp
--- a/Engine/Source/Editor/GameplayTagsEditor/Private/GameplayTagCustomization.cpp
+++ b/Engine/Source/Editor/GameplayTagsEditor/Private/GameplayTagCustomization.cpp
@@ -169,7 +169,9 @@ void FGameplayTagCustomization::BuildEditableContainerList()
 			TagName = Tag->ToString();
 			TagContainer->AddTag(*Tag);
 		}
-		EditableContainers.Add(SGameplayTagWidget::FEditableGameplayTagContainerDatum(OuterObjects[0], TagContainer.Get()));
+
+		UObject* TagContainerOwner = OuterObjects.Num() ? OuterObjects[0] : nullptr;
+		EditableContainers.Add(SGameplayTagWidget::FEditableGameplayTagContainerDatum(TagContainerOwner, TagContainer.Get()));
 	}
 }
 
diff --git a/Engine/Source/Editor/GameplayTagsEditor/Private/SGameplayTagWidget.cpp b/Engine/Source/Editor/GameplayTagsEditor/Private/SGameplayTagWidget.cpp
--- a/Engine/Source/Editor/GameplayTagsEditor/Private/SGameplayTagWidget.cpp
+++ b/Engine/Source/Editor/GameplayTagsEditor/Private/SGameplayTagWidget.cpp
@@ -223,38 +223,35 @@ void SGameplayTagWidget::OnTagChecked(TSharedPtr<FGameplayTagNode> NodeChecked)
 
 	for (int32 ContainerIdx = 0; ContainerIdx < TagContainers.Num(); ++ContainerIdx)
 	{
-		if (TagContainers[ContainerIdx].TagContainerOwner)
+		TWeakPtr<FGameplayTagNode> CurNode(NodeChecked);
+		UObject* OwnerObj = TagContainers[ContainerIdx].TagContainerOwner;
+		FGameplayTagContainer* Container = TagContainers[ContainerIdx].TagContainer;
+
+		if (Container)
 		{
-			TWeakPtr<FGameplayTagNode> CurNode(NodeChecked);
-			UObject* OwnerObj = TagContainers[ContainerIdx].TagContainerOwner;
-			FGameplayTagContainer* Container = TagContainers[ContainerIdx].TagContainer;
+			FGameplayTagContainer EditableContainer = *Container;
 
-			if (OwnerObj && Container)
+			while (CurNode.IsValid())
 			{
-				FGameplayTagContainer EditableContainer = *Container;
+				FGameplayTag Tag = TagsManager.RequestGameplayTag(CurNode.Pin()->GetCompleteTag());
 
-				while (CurNode.IsValid())
+				if (bRemoveParents == false)
 				{
-					FGameplayTag Tag = TagsManager.RequestGameplayTag(CurNode.Pin()->GetCompleteTag());
-
-					if (bRemoveParents == false)
+					bRemoveParents = true;
+					if (bMultiSelect == false)
 					{
-						bRemoveParents = true;
-						if (bMultiSelect == false)
-						{
-							EditableContainer.RemoveAllTags();
-						}
-						EditableContainer.AddTag(Tag);
-					}
-					else
-					{
-						EditableContainer.RemoveTag(Tag);
+						EditableContainer.RemoveAllTags();
 					}
-
-					CurNode = CurNode.Pin()->GetParentTagNode();
+					EditableContainer.AddTag(Tag);
 				}
-				SetContainer(Container, &EditableContainer, OwnerObj);
+				else
+				{
+					EditableContainer.RemoveTag(Tag);
+				}
+
+				CurNode = CurNode.Pin()->GetParentTagNode();
 			}
+			SetContainer(Container, &EditableContainer, OwnerObj);
 		}
 	}
 }
@@ -272,7 +269,7 @@ void SGameplayTagWidget::OnTagUnchecked(TSharedPtr<FGameplayTagNode> NodeUncheck
 			FGameplayTagContainer* Container = TagContainers[ContainerIdx].TagContainer;
 			FGameplayTag Tag = TagsManager.RequestGameplayTag(NodeUnchecked->GetCompleteTag());
 
-			if (OwnerObj && Container)
+			if (Container)
 			{
 				FGameplayTagContainer EditableContainer = *Container;
 				EditableContainer.RemoveTag(Tag);
@@ -382,7 +379,7 @@ FReply SGameplayTagWidget::OnClearAllClicked()
 		UObject* OwnerObj = TagContainers[ContainerIdx].TagContainerOwner;
 		FGameplayTagContainer* Container = TagContainers[ContainerIdx].TagContainer;
 
-		if (OwnerObj && Container)
+		if (Container)
 		{
 			FGameplayTagContainer EmptyContainer;
 			SetContainer(Container, &EmptyContainer, OwnerObj);
@@ -455,7 +452,7 @@ void SGameplayTagWidget::VerifyAssetTagValidity()
 		UObject* OwnerObj = TagContainers[ContainerIdx].TagContainerOwner;
 		FGameplayTagContainer* Container = TagContainers[ContainerIdx].TagContainer;
 
-		if (OwnerObj && Container)
+		if (Container)
 		{
 			FGameplayTagContainer EditableContainer = *Container;
 			FGameplayTagContainer InvalidTags;
@@ -547,31 +544,39 @@ void SGameplayTagWidget::OnExpansionChanged( TSharedPtr<FGameplayTagNode> InItem
 
 void SGameplayTagWidget::SetContainer(FGameplayTagContainer* OriginalContainer, FGameplayTagContainer* EditedContainer, UObject* OwnerObj)
 {
+	if (OwnerObj)
+	{
+		OwnerObj->PreEditChange(PropertyHandle.IsValid() ? PropertyHandle->GetProperty() : NULL);
+	}
+
 	if (PropertyHandle.IsValid() && bMultiSelect)
 	{
 		// Case for a tag container 
-		OwnerObj->PreEditChange(PropertyHandle->GetProperty());
 		PropertyHandle->SetValueFromFormattedString(EditedContainer->ToString());
-		OwnerObj->PostEditChange();
 	}
 	else if (PropertyHandle.IsValid() && !bMultiSelect)
 	{
 		// Case for a single Tag		
-		OwnerObj->PreEditChange(PropertyHandle->GetProperty());
 		FString FormattedString = TEXT("(TagName=\"");
 		FormattedString += EditedContainer->First().GetTagName().ToString();
 		FormattedString += TEXT("\")");
 		PropertyHandle->SetValueFromFormattedString(FormattedString);
-		OwnerObj->PostEditChange();
 	}
 	else
 	{
 		// Not sure if we should get here, means the property handle hasnt been setup which could be right or wrong.
-		OwnerObj->PreEditChange(PropertyHandle.IsValid() ? PropertyHandle->GetProperty() : NULL);
 		*OriginalContainer = *EditedContainer;
+	}	
+
+	if (OwnerObj)
+	{
 		OwnerObj->PostEditChange();
+	}
+
+	if (!PropertyHandle.IsValid())
+	{
 		OnTagChanged.ExecuteIfBound();
-	}	
+	}
 }
 
 #undef LOCTEXT_NAMESPACE
