Modified the SGameplayTagWidget to scale content with zoom so it can always display all its content when grapheditors have zoom applied.
#UE-13193 - Menus inaccessible based on zoom factor
#Branch UE4
#Proj GameplayTagsEditor

diff --git a/Engine/Source/Editor/GameplayTagsEditor/Private/SGameplayTagWidget.cpp b/Engine/Source/Editor/GameplayTagsEditor/Private/SGameplayTagWidget.cpp
--- a/Engine/Source/Editor/GameplayTagsEditor/Private/SGameplayTagWidget.cpp
+++ b/Engine/Source/Editor/GameplayTagsEditor/Private/SGameplayTagWidget.cpp
@@ -7,6 +7,7 @@
 #include "ScopedTransaction.h"
 #include "Editor/PropertyEditor/Public/PropertyHandle.h"
 #include "SSearchBox.h"
+#include "SScaleBox.h"
 
 #define LOCTEXT_NAMESPACE "GameplayTagWidget"
 
@@ -37,58 +38,66 @@ void SGameplayTagWidget::Construct(const FArguments& InArgs, const TArray<FEdita
 
 	ChildSlot
 	[
-		SNew(SBorder)
-		.BorderImage(FEditorStyle::GetBrush("ToolPanel.GroupBorder"))
+		SNew(SScaleBox)
+		.HAlign(EHorizontalAlignment::HAlign_Left)
+		.VAlign(EVerticalAlignment::VAlign_Top)
+		.StretchDirection(EStretchDirection::DownOnly)
+		.Stretch(EStretch::ScaleToFit)
+		.Content()
 		[
-			SNew(SVerticalBox)
-			+SVerticalBox::Slot()
-			.AutoHeight()
-			.VAlign(VAlign_Top)
+			SNew(SBorder)
+			.BorderImage(FEditorStyle::GetBrush("ToolPanel.GroupBorder"))
 			[
-				SNew(SHorizontalBox)
-				+SHorizontalBox::Slot()
-				.AutoWidth()
+				SNew(SVerticalBox)
+				+SVerticalBox::Slot()
+				.AutoHeight()
+				.VAlign(VAlign_Top)
 				[
-					SNew(SButton)
-					.OnClicked(this, &SGameplayTagWidget::OnExpandAllClicked)
-					.Text(LOCTEXT("GameplayTagWidget_ExpandAll", "Expand All"))
+					SNew(SHorizontalBox)
+					+SHorizontalBox::Slot()
+					.AutoWidth()
+					[
+						SNew(SButton)
+						.OnClicked(this, &SGameplayTagWidget::OnExpandAllClicked)
+						.Text(LOCTEXT("GameplayTagWidget_ExpandAll", "Expand All"))
+					]
+					+SHorizontalBox::Slot()
+					.AutoWidth()
+					[
+						SNew(SButton)
+						.OnClicked(this, &SGameplayTagWidget::OnCollapseAllClicked)
+						.Text(LOCTEXT("GameplayTagWidget_CollapseAll", "Collapse All"))
+					]
+					+SHorizontalBox::Slot()
+					.AutoWidth()
+					[
+						SNew(SButton)
+						.IsEnabled(!bReadOnly)
+						.OnClicked(this, &SGameplayTagWidget::OnClearAllClicked)
+						.Text(LOCTEXT("GameplayTagWidget_ClearAll", "Clear All"))
+					]
+					+SHorizontalBox::Slot()
+					.VAlign( VAlign_Center )
+					.FillWidth(1.f)
+					.Padding(5,1,5,1)
+					[
+						SNew(SSearchBox)
+						.HintText(LOCTEXT("GameplayTagWidget_SearchBoxHint", "Search Gameplay Tags"))
+						.OnTextChanged( this, &SGameplayTagWidget::OnFilterTextChanged )
+					]
 				]
-				+SHorizontalBox::Slot()
-				.AutoWidth()
+				+SVerticalBox::Slot()
 				[
-					SNew(SButton)
-					.OnClicked(this, &SGameplayTagWidget::OnCollapseAllClicked)
-					.Text(LOCTEXT("GameplayTagWidget_CollapseAll", "Collapse All"))
-				]
-				+SHorizontalBox::Slot()
-				.AutoWidth()
-				[
-					SNew(SButton)
-					.IsEnabled(!bReadOnly)
-					.OnClicked(this, &SGameplayTagWidget::OnClearAllClicked)
-					.Text(LOCTEXT("GameplayTagWidget_ClearAll", "Clear All"))
-				]
-				+SHorizontalBox::Slot()
-				.VAlign( VAlign_Center )
-				.FillWidth(1.f)
-				.Padding(5,1,5,1)
-				[
-					SNew(SSearchBox)
-					.HintText(LOCTEXT("GameplayTagWidget_SearchBoxHint", "Search Gameplay Tags"))
-					.OnTextChanged( this, &SGameplayTagWidget::OnFilterTextChanged )
-				]
-			]
-			+SVerticalBox::Slot()
-			[
-				SNew(SBorder)
-				.Padding(FMargin(4.f))
-				[
-					SAssignNew(TagTreeWidget, STreeView< TSharedPtr<FGameplayTagNode> >)
-					.TreeItemsSource(&TagItems)
-					.OnGenerateRow(this, &SGameplayTagWidget::OnGenerateRow)
-					.OnGetChildren(this, &SGameplayTagWidget::OnGetChildren)
-					.OnExpansionChanged( this, &SGameplayTagWidget::OnExpansionChanged)
-					.SelectionMode(ESelectionMode::Multi)
+					SNew(SBorder)
+					.Padding(FMargin(4.f))
+					[
+						SAssignNew(TagTreeWidget, STreeView< TSharedPtr<FGameplayTagNode> >)
+						.TreeItemsSource(&TagItems)
+						.OnGenerateRow(this, &SGameplayTagWidget::OnGenerateRow)
+						.OnGetChildren(this, &SGameplayTagWidget::OnGetChildren)
+						.OnExpansionChanged( this, &SGameplayTagWidget::OnExpansionChanged)
+						.SelectionMode(ESelectionMode::Multi)
+					]
 				]
 			]
 		]
