Fixed GameplayTagQuery issue when trying to edit same query while already being edited

diff --git a/Engine/Source/Editor/GameplayTagsEditor/Private/GameplayTagQueryCustomization.cpp b/Engine/Source/Editor/GameplayTagsEditor/Private/GameplayTagQueryCustomization.cpp
--- a/Engine/Source/Editor/GameplayTagsEditor/Private/GameplayTagQueryCustomization.cpp
+++ b/Engine/Source/Editor/GameplayTagsEditor/Private/GameplayTagQueryCustomization.cpp
@@ -54,44 +54,53 @@ FText FGameplayTagQueryCustomization::GetEditButtonText() const
 
 FReply FGameplayTagQueryCustomization::OnEditButtonClicked()
 {
-	TArray<UObject*> OuterObjects;
-	StructPropertyHandle->GetOuterObjects(OuterObjects);
-
-	bool const bReadOnly = StructPropertyHandle->GetProperty()->HasAnyPropertyFlags( CPF_EditConst );
-
-	FText Title;
-	if (OuterObjects.Num() > 1)
+	if (GameplayTagQueryWidgetWindow.IsValid())
 	{
-		FText const AssetName = FText::Format( LOCTEXT("GameplayTagDetailsBase_MultipleAssets", "{0} Assets"), FText::AsNumber( OuterObjects.Num() ) );
-		FText const PropertyName = StructPropertyHandle->GetPropertyDisplayName();
-		Title = FText::Format(LOCTEXT("GameplayTagQueryCustomization_BaseWidgetTitle", "Tag Editor: {0} {1}"), PropertyName, AssetName);
+		// already open, just show it
+		GameplayTagQueryWidgetWindow->BringToFront(true);
 	}
-	else if (OuterObjects.Num() > 0 && OuterObjects[0])
+	else
 	{
-		FText const AssetName = FText::FromString( OuterObjects[0]->GetName() );
-		FText const PropertyName = StructPropertyHandle->GetPropertyDisplayName();
-		Title = FText::Format(LOCTEXT("GameplayTagQueryCustomization_BaseWidgetTitle", "Tag Editor: {0} {1}"), PropertyName, AssetName);
-	}
+		TArray<UObject*> OuterObjects;
+		StructPropertyHandle->GetOuterObjects(OuterObjects);
 
-	GameplayTagQueryWidgetWindow = SNew(SWindow)
-	.Title(Title)
-	.HasCloseButton(false)
-	.ClientSize(FVector2D(600, 400))
-	[
-		SNew(SGameplayTagQueryWidget, EditableQueries)
-		.OnSaveAndClose(this, &FGameplayTagQueryCustomization::CloseWidgetWindow)
-		.OnCancel(this, &FGameplayTagQueryCustomization::CloseWidgetWindow)
-		.ReadOnly(bReadOnly)
-	];
+		bool const bReadOnly = StructPropertyHandle->GetProperty()->HasAnyPropertyFlags(CPF_EditConst);
 
-	IMainFrameModule& MainFrameModule = FModuleManager::LoadModuleChecked<IMainFrameModule>(TEXT("MainFrame"));
-	if (MainFrameModule.GetParentWindow().IsValid())
-	{
-		FSlateApplication::Get().AddWindowAsNativeChild(GameplayTagQueryWidgetWindow.ToSharedRef(), MainFrameModule.GetParentWindow().ToSharedRef());
-	}
-	else
-	{
-		FSlateApplication::Get().AddWindow(GameplayTagQueryWidgetWindow.ToSharedRef());
+		FText Title;
+		if (OuterObjects.Num() > 1)
+		{
+			FText const AssetName = FText::Format(LOCTEXT("GameplayTagDetailsBase_MultipleAssets", "{0} Assets"), FText::AsNumber(OuterObjects.Num()));
+			FText const PropertyName = StructPropertyHandle->GetPropertyDisplayName();
+			Title = FText::Format(LOCTEXT("GameplayTagQueryCustomization_BaseWidgetTitle", "Tag Editor: {0} {1}"), PropertyName, AssetName);
+		}
+		else if (OuterObjects.Num() > 0 && OuterObjects[0])
+		{
+			FText const AssetName = FText::FromString(OuterObjects[0]->GetName());
+			FText const PropertyName = StructPropertyHandle->GetPropertyDisplayName();
+			Title = FText::Format(LOCTEXT("GameplayTagQueryCustomization_BaseWidgetTitle", "Tag Editor: {0} {1}"), PropertyName, AssetName);
+		}
+
+
+		GameplayTagQueryWidgetWindow = SNew(SWindow)
+			.Title(Title)
+			.HasCloseButton(false)
+			.ClientSize(FVector2D(600, 400))
+			[
+				SNew(SGameplayTagQueryWidget, EditableQueries)
+				.OnSaveAndClose(this, &FGameplayTagQueryCustomization::CloseWidgetWindow)
+				.OnCancel(this, &FGameplayTagQueryCustomization::CloseWidgetWindow)
+				.ReadOnly(bReadOnly)
+			];
+
+		IMainFrameModule& MainFrameModule = FModuleManager::LoadModuleChecked<IMainFrameModule>(TEXT("MainFrame"));
+		if (MainFrameModule.GetParentWindow().IsValid())
+		{
+			FSlateApplication::Get().AddWindowAsNativeChild(GameplayTagQueryWidgetWindow.ToSharedRef(), MainFrameModule.GetParentWindow().ToSharedRef());
+		}
+		else
+		{
+			FSlateApplication::Get().AddWindow(GameplayTagQueryWidgetWindow.ToSharedRef());
+		}
 	}
 
 	return FReply::Handled();
@@ -131,6 +140,7 @@ void FGameplayTagQueryCustomization::CloseWidgetWindow()
  	if( GameplayTagQueryWidgetWindow.IsValid() )
  	{
  		GameplayTagQueryWidgetWindow->RequestDestroyWindow();
+		GameplayTagQueryWidgetWindow = nullptr;
  	}
 }
 
