Fix shadow variable

diff --git a/Engine/Source/Editor/GameplayTagsEditor/Private/SGameplayTagQueryWidget.cpp b/Engine/Source/Editor/GameplayTagsEditor/Private/SGameplayTagQueryWidget.cpp
--- a/Engine/Source/Editor/GameplayTagsEditor/Private/SGameplayTagQueryWidget.cpp
+++ b/Engine/Source/Editor/GameplayTagsEditor/Private/SGameplayTagQueryWidget.cpp
@@ -91,14 +91,14 @@ void SGameplayTagQueryWidget::Construct(const FArguments& InArgs, const TArray<F
 
 UEditableGameplayTagQuery* SGameplayTagQueryWidget::CreateEditableQuery(FGameplayTagQuery& Q)
 {
-	UEditableGameplayTagQuery* const EditableQuery = Q.CreateEditableQuery();
-	if (EditableQuery)
+	UEditableGameplayTagQuery* const AnEditableQuery = Q.CreateEditableQuery();
+	if (AnEditableQuery)
 	{
 		// prevent GC, will explicitly remove from root later
-		EditableQuery->AddToRoot();
+		AnEditableQuery->AddToRoot();
 	}
 
-	return EditableQuery;
+	return AnEditableQuery;
 }
 
 SGameplayTagQueryWidget::~SGameplayTagQueryWidget()
