Fixed potential crash in GameplayTagQueryCustomization (UE-18459)

diff --git a/Engine/Source/Editor/GameplayTagsEditor/Private/GameplayTagQueryCustomization.cpp b/Engine/Source/Editor/GameplayTagsEditor/Private/GameplayTagQueryCustomization.cpp
--- a/Engine/Source/Editor/GameplayTagsEditor/Private/GameplayTagQueryCustomization.cpp
+++ b/Engine/Source/Editor/GameplayTagsEditor/Private/GameplayTagQueryCustomization.cpp
@@ -16,6 +16,8 @@ void FGameplayTagQueryCustomization::CustomizeHeader(TSharedRef<class IPropertyH
 
 	BuildEditableQueryList();
 
+	bool const bReadOnly = StructPropertyHandle->GetProperty() ? StructPropertyHandle->GetProperty()->HasAnyPropertyFlags(CPF_EditConst) : false;
+
 	HeaderRow
 	.NameContent()
 	[
@@ -41,7 +43,7 @@ void FGameplayTagQueryCustomization::CustomizeHeader(TSharedRef<class IPropertyH
 			.AutoHeight()
 			[
 				SNew(SButton)
-				.IsEnabled(!StructPropertyHandle->GetProperty()->HasAnyPropertyFlags(CPF_EditConst))
+				.IsEnabled(!bReadOnly)
 				.Text(LOCTEXT("GameplayTagQueryCustomization_Clear", "Clear All"))
 				.OnClicked(this, &FGameplayTagQueryCustomization::OnClearAllButtonClicked)
 				.Visibility(this, &FGameplayTagQueryCustomization::GetClearAllVisibility)
@@ -71,11 +73,17 @@ FText FGameplayTagQueryCustomization::GetQueryDescText() const
 
 FText FGameplayTagQueryCustomization::GetEditButtonText() const
 {
-	bool const bReadOnly = StructPropertyHandle->GetProperty()->HasAnyPropertyFlags(CPF_EditConst);
-	return
-		bReadOnly
-		? LOCTEXT("GameplayTagQueryCustomization_View", "View...")
-		: LOCTEXT("GameplayTagQueryCustomization_Edit", "Edit...");
+	UProperty const* const Prop = StructPropertyHandle.IsValid() ? StructPropertyHandle->GetProperty() : nullptr;
+	if (Prop)
+	{
+		bool const bReadOnly = Prop->HasAnyPropertyFlags(CPF_EditConst);
+		return
+			bReadOnly
+			? LOCTEXT("GameplayTagQueryCustomization_View", "View...")
+			: LOCTEXT("GameplayTagQueryCustomization_Edit", "Edit...");
+	}
+
+	return FText();
 }
 
 FReply FGameplayTagQueryCustomization::OnClearAllButtonClicked()
@@ -145,7 +153,7 @@ FReply FGameplayTagQueryCustomization::OnEditButtonClicked()
 		TArray<UObject*> OuterObjects;
 		StructPropertyHandle->GetOuterObjects(OuterObjects);
 
-		bool const bReadOnly = StructPropertyHandle->GetProperty()->HasAnyPropertyFlags(CPF_EditConst);
+		bool const bReadOnly = StructPropertyHandle->GetProperty() ? StructPropertyHandle->GetProperty()->HasAnyPropertyFlags(CPF_EditConst) : false;
 
 		FText Title;
 		if (OuterObjects.Num() > 1)
