Fixed Tag Editor window layout scaling that was causing input problems on Mac
[UE-18455] - Tag Editor window cannot be focused on Mac

diff --git a/Engine/Source/Editor/GameplayTagsEditor/Private/GameplayTagCustomization.cpp b/Engine/Source/Editor/GameplayTagsEditor/Private/GameplayTagCustomization.cpp
--- a/Engine/Source/Editor/GameplayTagsEditor/Private/GameplayTagCustomization.cpp
+++ b/Engine/Source/Editor/GameplayTagsEditor/Private/GameplayTagCustomization.cpp
@@ -7,6 +7,7 @@
 #include "SGameplayTagWidget.h"
 #include "GameplayTagContainer.h"
 #include "ScopedTransaction.h"
+#include "SScaleBox.h"
 
 #define LOCTEXT_NAMESPACE "GameplayTagCustomization"
 
@@ -74,13 +75,21 @@ TSharedRef<SWidget> FGameplayTagCustomization::GetListContent()
 		.AutoHeight()
 		.MaxHeight(400)
 		[
-			SNew(SGameplayTagWidget, EditableContainers)
-			.Filter(Categories)
-			.ReadOnly(bReadOnly)
-			.TagContainerName(StructPropertyHandle->GetPropertyDisplayName().ToString())
-			.MultiSelect(false)
-			.OnTagChanged(this, &FGameplayTagCustomization::OnTagChanged)
-			.PropertyHandle(StructPropertyHandle)
+			SNew(SScaleBox)
+			.HAlign(EHorizontalAlignment::HAlign_Left)
+			.VAlign(EVerticalAlignment::VAlign_Top)
+			.StretchDirection(EStretchDirection::DownOnly)
+			.Stretch(EStretch::ScaleToFit)
+			.Content()
+			[
+				SNew(SGameplayTagWidget, EditableContainers)
+				.Filter(Categories)
+				.ReadOnly(bReadOnly)
+				.TagContainerName(StructPropertyHandle->GetPropertyDisplayName().ToString())
+				.MultiSelect(false)
+				.OnTagChanged(this, &FGameplayTagCustomization::OnTagChanged)
+				.PropertyHandle(StructPropertyHandle)
+			]
 		];
 }
 
diff --git a/Engine/Source/Editor/GameplayTagsEditor/Private/SGameplayTagContainerGraphPin.cpp b/Engine/Source/Editor/GameplayTagsEditor/Private/SGameplayTagContainerGraphPin.cpp
--- a/Engine/Source/Editor/GameplayTagsEditor/Private/SGameplayTagContainerGraphPin.cpp
+++ b/Engine/Source/Editor/GameplayTagsEditor/Private/SGameplayTagContainerGraphPin.cpp
@@ -4,6 +4,7 @@
 #include "SGameplayTagContainerGraphPin.h"
 #include "GameplayTagsModule.h"
 #include "GameplayTags.h"
+#include "SScaleBox.h"
 
 #define LOCTEXT_NAMESPACE "GameplayTagGraphPin"
 
@@ -104,10 +105,18 @@ TSharedRef<SWidget> SGameplayTagContainerGraphPin::GetListContent()
 		.AutoHeight()
 		.MaxHeight( 400 )
 		[
-			SNew( SGameplayTagWidget, EditableContainers )
-			.OnTagChanged(this, &SGameplayTagContainerGraphPin::RefreshTagList)
-			.TagContainerName( TEXT("SGameplayTagContainerGraphPin") )
-			.Visibility( this, &SGraphPin::GetDefaultValueVisibility )
+			SNew(SScaleBox)
+			.HAlign(EHorizontalAlignment::HAlign_Left)
+			.VAlign(EVerticalAlignment::VAlign_Top)
+			.StretchDirection(EStretchDirection::DownOnly)
+			.Stretch(EStretch::ScaleToFit)
+			.Content()
+			[
+				SNew( SGameplayTagWidget, EditableContainers )
+				.OnTagChanged(this, &SGameplayTagContainerGraphPin::RefreshTagList)
+				.TagContainerName( TEXT("SGameplayTagContainerGraphPin") )
+				.Visibility( this, &SGraphPin::GetDefaultValueVisibility )
+			]
 		];
 }
 
diff --git a/Engine/Source/Editor/GameplayTagsEditor/Private/SGameplayTagGraphPin.cpp b/Engine/Source/Editor/GameplayTagsEditor/Private/SGameplayTagGraphPin.cpp
--- a/Engine/Source/Editor/GameplayTagsEditor/Private/SGameplayTagGraphPin.cpp
+++ b/Engine/Source/Editor/GameplayTagsEditor/Private/SGameplayTagGraphPin.cpp
@@ -4,6 +4,7 @@
 #include "SGameplayTagGraphPin.h"
 #include "GameplayTagsModule.h"
 #include "GameplayTags.h"
+#include "SScaleBox.h"
 
 #define LOCTEXT_NAMESPACE "GameplayTagGraphPin"
 
@@ -88,12 +89,20 @@ TSharedRef<SWidget> SGameplayTagGraphPin::GetListContent()
 		.AutoHeight()
 		.MaxHeight( 400 )
 		[
-			SNew( SGameplayTagWidget, EditableContainers )
-			.OnTagChanged( this, &SGameplayTagGraphPin::RefreshTagList )
-			.TagContainerName( TEXT("SGameplayTagGraphPin") )
-			.Visibility( this, &SGraphPin::GetDefaultValueVisibility )
-			.MultiSelect(false)
-			.Filter(FilterString)
+			SNew(SScaleBox)
+			.HAlign(EHorizontalAlignment::HAlign_Left)
+			.VAlign(EVerticalAlignment::VAlign_Top)
+			.StretchDirection(EStretchDirection::DownOnly)
+			.Stretch(EStretch::ScaleToFit)
+			.Content()
+			[
+				SNew( SGameplayTagWidget, EditableContainers )
+				.OnTagChanged( this, &SGameplayTagGraphPin::RefreshTagList )
+				.TagContainerName( TEXT("SGameplayTagGraphPin") )
+				.Visibility( this, &SGraphPin::GetDefaultValueVisibility )
+				.MultiSelect(false)
+				.Filter(FilterString)
+			]
 		];
 }
 
diff --git a/Engine/Source/Editor/GameplayTagsEditor/Private/SGameplayTagQueryGraphPin.cpp b/Engine/Source/Editor/GameplayTagsEditor/Private/SGameplayTagQueryGraphPin.cpp
--- a/Engine/Source/Editor/GameplayTagsEditor/Private/SGameplayTagQueryGraphPin.cpp
+++ b/Engine/Source/Editor/GameplayTagsEditor/Private/SGameplayTagQueryGraphPin.cpp
@@ -4,6 +4,7 @@
 #include "SGameplayTagQueryGraphPin.h"
 #include "GameplayTagsModule.h"
 #include "GameplayTags.h"
+#include "SScaleBox.h"
 
 #define LOCTEXT_NAMESPACE "GameplayTagQueryGraphPin"
 
@@ -61,10 +62,18 @@ TSharedRef<SWidget> SGameplayTagQueryGraphPin::GetListContent()
 		.AutoHeight()
 		.MaxHeight( 400 )
 		[
-			SNew(SGameplayTagQueryWidget, EditableQueries)
-			.OnQueryChanged(this, &SGameplayTagQueryGraphPin::OnQueryChanged)
-			.Visibility( this, &SGraphPin::GetDefaultValueVisibility )
-			.AutoSave(true)
+			SNew(SScaleBox)
+			.HAlign(EHorizontalAlignment::HAlign_Left)
+			.VAlign(EVerticalAlignment::VAlign_Top)
+			.StretchDirection(EStretchDirection::DownOnly)
+			.Stretch(EStretch::ScaleToFit)
+			.Content()
+			[
+				SNew(SGameplayTagQueryWidget, EditableQueries)
+				.OnQueryChanged(this, &SGameplayTagQueryGraphPin::OnQueryChanged)
+				.Visibility( this, &SGraphPin::GetDefaultValueVisibility )
+				.AutoSave(true)
+			]
 		];
 }
 
diff --git a/Engine/Source/Editor/GameplayTagsEditor/Private/SGameplayTagQueryWidget.cpp b/Engine/Source/Editor/GameplayTagsEditor/Private/SGameplayTagQueryWidget.cpp
--- a/Engine/Source/Editor/GameplayTagsEditor/Private/SGameplayTagQueryWidget.cpp
+++ b/Engine/Source/Editor/GameplayTagsEditor/Private/SGameplayTagQueryWidget.cpp
@@ -50,46 +50,38 @@ void SGameplayTagQueryWidget::Construct(const FArguments& InArgs, const TArray<F
 
 	ChildSlot
 	[
-		SNew(SScaleBox)
-		.HAlign(EHorizontalAlignment::HAlign_Left)
-		.VAlign(EVerticalAlignment::VAlign_Top)
-		.StretchDirection(EStretchDirection::DownOnly)
-		.Stretch(EStretch::ScaleToFit)
-		.Content()
+		SNew(SBorder)
+		.BorderImage(FEditorStyle::GetBrush("ToolPanel.GroupBorder"))
 		[
-			SNew(SBorder)
-			.BorderImage(FEditorStyle::GetBrush("ToolPanel.GroupBorder"))
+			SNew(SVerticalBox)
+			+ SVerticalBox::Slot()
+			.AutoHeight()
+			.VAlign(VAlign_Top)
 			[
-				SNew(SVerticalBox)
-				+ SVerticalBox::Slot()
-				.AutoHeight()
-				.VAlign(VAlign_Top)
+				SNew(SHorizontalBox)
+				+ SHorizontalBox::Slot()
+				.AutoWidth()
 				[
-					SNew(SHorizontalBox)
-					+ SHorizontalBox::Slot()
-					.AutoWidth()
-					[
-						SNew(SButton)
-						.IsEnabled(!bReadOnly)
-						.Visibility(this, &SGameplayTagQueryWidget::GetSaveAndCloseButtonVisibility)
-						.OnClicked(this, &SGameplayTagQueryWidget::OnSaveAndCloseClicked)
-						.Text(LOCTEXT("GameplayTagQueryWidget_SaveAndClose", "Save and Close"))
-					]
-					+ SHorizontalBox::Slot()
-					.AutoWidth()
-					[
-						SNew(SButton)
-						.Visibility(this, &SGameplayTagQueryWidget::GetCancelButtonVisibility)
-						.OnClicked(this, &SGameplayTagQueryWidget::OnCancelClicked)
-						.Text(LOCTEXT("GameplayTagQueryWidget_Cancel", "Close Without Saving"))
-					]
+					SNew(SButton)
+					.IsEnabled(!bReadOnly)
+					.Visibility(this, &SGameplayTagQueryWidget::GetSaveAndCloseButtonVisibility)
+					.OnClicked(this, &SGameplayTagQueryWidget::OnSaveAndCloseClicked)
+					.Text(LOCTEXT("GameplayTagQueryWidget_SaveAndClose", "Save and Close"))
 				]
-				// to delete!
-				+ SVerticalBox::Slot()
+				+ SHorizontalBox::Slot()
+				.AutoWidth()
 				[
-					Details.ToSharedRef()
+					SNew(SButton)
+					.Visibility(this, &SGameplayTagQueryWidget::GetCancelButtonVisibility)
+					.OnClicked(this, &SGameplayTagQueryWidget::OnCancelClicked)
+					.Text(LOCTEXT("GameplayTagQueryWidget_Cancel", "Close Without Saving"))
 				]
 			]
+			// to delete!
+			+ SVerticalBox::Slot()
+			[
+				Details.ToSharedRef()
+			]
 		]
 	];
 }
diff --git a/Engine/Source/Editor/GameplayTagsEditor/Private/SGameplayTagWidget.cpp b/Engine/Source/Editor/GameplayTagsEditor/Private/SGameplayTagWidget.cpp
--- a/Engine/Source/Editor/GameplayTagsEditor/Private/SGameplayTagWidget.cpp
+++ b/Engine/Source/Editor/GameplayTagsEditor/Private/SGameplayTagWidget.cpp
@@ -38,66 +38,58 @@ void SGameplayTagWidget::Construct(const FArguments& InArgs, const TArray<FEdita
 
 	ChildSlot
 	[
-		SNew(SScaleBox)
-		.HAlign(EHorizontalAlignment::HAlign_Left)
-		.VAlign(EVerticalAlignment::VAlign_Top)
-		.StretchDirection(EStretchDirection::DownOnly)
-		.Stretch(EStretch::ScaleToFit)
-		.Content()
+		SNew(SBorder)
+		.BorderImage(FEditorStyle::GetBrush("ToolPanel.GroupBorder"))
 		[
-			SNew(SBorder)
-			.BorderImage(FEditorStyle::GetBrush("ToolPanel.GroupBorder"))
+			SNew(SVerticalBox)
+			+SVerticalBox::Slot()
+			.AutoHeight()
+			.VAlign(VAlign_Top)
 			[
-				SNew(SVerticalBox)
-				+SVerticalBox::Slot()
-				.AutoHeight()
-				.VAlign(VAlign_Top)
+				SNew(SHorizontalBox)
+				+SHorizontalBox::Slot()
+				.AutoWidth()
 				[
-					SNew(SHorizontalBox)
-					+SHorizontalBox::Slot()
-					.AutoWidth()
-					[
-						SNew(SButton)
-						.OnClicked(this, &SGameplayTagWidget::OnExpandAllClicked)
-						.Text(LOCTEXT("GameplayTagWidget_ExpandAll", "Expand All"))
-					]
-					+SHorizontalBox::Slot()
-					.AutoWidth()
-					[
-						SNew(SButton)
-						.OnClicked(this, &SGameplayTagWidget::OnCollapseAllClicked)
-						.Text(LOCTEXT("GameplayTagWidget_CollapseAll", "Collapse All"))
-					]
-					+SHorizontalBox::Slot()
-					.AutoWidth()
-					[
-						SNew(SButton)
-						.IsEnabled(!bReadOnly)
-						.OnClicked(this, &SGameplayTagWidget::OnClearAllClicked)
-						.Text(LOCTEXT("GameplayTagWidget_ClearAll", "Clear All"))
-					]
-					+SHorizontalBox::Slot()
-					.VAlign( VAlign_Center )
-					.FillWidth(1.f)
-					.Padding(5,1,5,1)
-					[
-						SNew(SSearchBox)
-						.HintText(LOCTEXT("GameplayTagWidget_SearchBoxHint", "Search Gameplay Tags"))
-						.OnTextChanged( this, &SGameplayTagWidget::OnFilterTextChanged )
-					]
+					SNew(SButton)
+					.OnClicked(this, &SGameplayTagWidget::OnExpandAllClicked)
+					.Text(LOCTEXT("GameplayTagWidget_ExpandAll", "Expand All"))
 				]
-				+SVerticalBox::Slot()
+				+SHorizontalBox::Slot()
+				.AutoWidth()
 				[
-					SNew(SBorder)
-					.Padding(FMargin(4.f))
-					[
-						SAssignNew(TagTreeWidget, STreeView< TSharedPtr<FGameplayTagNode> >)
-						.TreeItemsSource(&TagItems)
-						.OnGenerateRow(this, &SGameplayTagWidget::OnGenerateRow)
-						.OnGetChildren(this, &SGameplayTagWidget::OnGetChildren)
-						.OnExpansionChanged( this, &SGameplayTagWidget::OnExpansionChanged)
-						.SelectionMode(ESelectionMode::Multi)
-					]
+					SNew(SButton)
+					.OnClicked(this, &SGameplayTagWidget::OnCollapseAllClicked)
+					.Text(LOCTEXT("GameplayTagWidget_CollapseAll", "Collapse All"))
+				]
+				+SHorizontalBox::Slot()
+				.AutoWidth()
+				[
+					SNew(SButton)
+					.IsEnabled(!bReadOnly)
+					.OnClicked(this, &SGameplayTagWidget::OnClearAllClicked)
+					.Text(LOCTEXT("GameplayTagWidget_ClearAll", "Clear All"))
+				]
+				+SHorizontalBox::Slot()
+				.VAlign( VAlign_Center )
+				.FillWidth(1.f)
+				.Padding(5,1,5,1)
+				[
+					SNew(SSearchBox)
+					.HintText(LOCTEXT("GameplayTagWidget_SearchBoxHint", "Search Gameplay Tags"))
+					.OnTextChanged( this, &SGameplayTagWidget::OnFilterTextChanged )
+				]
+			]
+			+SVerticalBox::Slot()
+			[
+				SNew(SBorder)
+				.Padding(FMargin(4.f))
+				[
+					SAssignNew(TagTreeWidget, STreeView< TSharedPtr<FGameplayTagNode> >)
+					.TreeItemsSource(&TagItems)
+					.OnGenerateRow(this, &SGameplayTagWidget::OnGenerateRow)
+					.OnGetChildren(this, &SGameplayTagWidget::OnGetChildren)
+					.OnExpansionChanged( this, &SGameplayTagWidget::OnExpansionChanged)
+					.SelectionMode(ESelectionMode::Multi)
 				]
 			]
 		]
