Gameplay Tags: Added a prototype content browser search filter for assets that use a specific gameplay tag
- Note: Only works with loaded assets currently, and the tag filter is an arbitrary text box

diff --git a/Engine/Source/Editor/GameplayTagsEditor/GameplayTagsEditor.Build.cs b/Engine/Source/Editor/GameplayTagsEditor/GameplayTagsEditor.Build.cs
--- a/Engine/Source/Editor/GameplayTagsEditor/GameplayTagsEditor.Build.cs
+++ b/Engine/Source/Editor/GameplayTagsEditor/GameplayTagsEditor.Build.cs
@@ -38,6 +38,7 @@ namespace UnrealBuildTool.Rules
 					"BlueprintGraph",
 					"KismetCompiler",
 					"GraphEditor",
+					"ContentBrowser",
 					"MainFrame",
 					"UnrealEd",
 				}
diff --git a/Engine/Source/Editor/GameplayTagsEditor/Private/GameplayTagSearchFilter.cpp b/Engine/Source/Editor/GameplayTagsEditor/Private/GameplayTagSearchFilter.cpp
new file mode 100644
--- /dev/null
+++ b/Engine/Source/Editor/GameplayTagsEditor/Private/GameplayTagSearchFilter.cpp
@@ -0,0 +1,203 @@
+// Copyright 1998-2015 Epic Games, Inc. All Rights Reserved.
+
+#pragma once
+
+#include "GameplayTagsEditorModulePrivatePCH.h"
+#include "GameplayTagSearchFilter.h"
+
+#include "SGameplayTagWidget.h"
+#include "GameplayTagsModule.h"
+
+#include "Editor/ContentBrowser/Public/ContentBrowserModule.h"
+#include "Editor/ContentBrowser/Public/FrontendFilterBase.h"
+
+#define LOCTEXT_NAMESPACE "GameplayTagSearchFilter"
+
+//////////////////////////////////////////////////////////////////////////
+//
+
+/** A filter that search for assets using a specific gameplay tag */
+class FFrontendFilter_GameplayTags : public FFrontendFilter
+{
+public:
+	FFrontendFilter_GameplayTags(TSharedPtr<FFrontendFilterCategory> InCategory) : FFrontendFilter(InCategory) {}
+
+	// FFrontendFilter implementation
+	virtual FString GetName() const override { return TEXT("GameplayTagFilter"); }
+	virtual FText GetDisplayName() const override;
+	virtual FText GetToolTipText() const override { return LOCTEXT("GameplayTagFilterDisplayTooltip", "Search for any Blueprint or asset containing a specified gameplay tag (right-click to change the tag)."); }
+	virtual void ModifyContextMenu(FMenuBuilder& MenuBuilder) override;
+	virtual void SaveSettings(const FString& IniFilename, const FString& IniSection, const FString& SettingsString) const override;
+	virtual void LoadSettings(const FString& IniFilename, const FString& IniSection, const FString& SettingsString) override;
+	// End of FFrontendFilter implementation
+
+	// IFilter implementation
+	virtual bool PassesFilter(FAssetFilterType InItem) const override;
+	// End of IFilter implementation
+
+protected:
+	// If blank, finds anything that contains a gameplay tag
+	FGameplayTag TagSearch;
+
+protected:
+	bool ProcessStruct(void* Data, UStruct* Struct) const;
+
+	bool ProcessProperty(void* Data, UProperty* Prop) const;
+
+	FText GetTagValueAsText() const;
+	void OnTagTextCommitted(const FText& InText, ETextCommit::Type InCommitType);
+	void OnTagWidgetChanged();
+};
+
+void FFrontendFilter_GameplayTags::ModifyContextMenu(FMenuBuilder& MenuBuilder)
+{
+	FUIAction Action;
+
+	MenuBuilder.BeginSection(TEXT("ComparsionSection"), LOCTEXT("ComparisonSectionHeading", "Gameplay Tag Comparison"));
+
+	TSharedRef<SWidget> KeyWidget =
+		SNew(SEditableTextBox)
+		.Text_Raw(this, &FFrontendFilter_GameplayTags::GetTagValueAsText)
+		.OnTextCommitted_Raw(this, &FFrontendFilter_GameplayTags::OnTagTextCommitted)
+		.MinDesiredWidth(100.0f);
+
+	MenuBuilder.AddWidget(KeyWidget, LOCTEXT("TagMenuDesc", "Gameplay Tag"));
+
+// 	TSharedRef<SWidget> TagWidget = SNew(SGameplayTagWidget)
+// 		.MultiSelect(true)
+// 		.OnTagChanged(this, &FFrontendFilter_GameplayTags::OnTagWidgetChanged);
+// 	MenuBuilder.AddWidget(TagWidget, LOCTEXT("TagMenuDesc2", "Gameplay Tag"));
+}
+
+FText FFrontendFilter_GameplayTags::GetDisplayName() const
+{
+	return FText::Format(LOCTEXT("FFrontendFilter_CompareOperation", "Gameplay Tags ({0})"),
+		TagSearch.IsValid() ? FText::FromName(TagSearch.GetTagName()) : LOCTEXT("AnyTag", "Any Tags"));
+}
+
+void FFrontendFilter_GameplayTags::SaveSettings(const FString& IniFilename, const FString& IniSection, const FString& SettingsString) const
+{
+	GConfig->SetString(*IniSection, *(SettingsString + TEXT(".Key")), *TagSearch.ToString(), IniFilename);
+}
+
+void FFrontendFilter_GameplayTags::LoadSettings(const FString& IniFilename, const FString& IniSection, const FString& SettingsString)
+{
+	FString TagNameAsString;
+	if (GConfig->GetString(*IniSection, *(SettingsString + TEXT(".Key")), TagNameAsString, IniFilename))
+	{
+		IGameplayTagsModule& GameplayTagModule = IGameplayTagsModule::Get();
+
+		FGameplayTag NewTag = GameplayTagModule.RequestGameplayTag(*TagNameAsString, /*bErrorIfNotFound=*/ false);
+		TagSearch = NewTag;
+	}
+}
+
+FText FFrontendFilter_GameplayTags::GetTagValueAsText() const
+{
+	return FText::FromName(TagSearch.GetTagName());
+}
+
+void FFrontendFilter_GameplayTags::OnTagTextCommitted(const FText& InText, ETextCommit::Type InCommitType)
+{
+	IGameplayTagsModule& GameplayTagModule = IGameplayTagsModule::Get();
+
+	FGameplayTag NewTag = GameplayTagModule.RequestGameplayTag(*InText.ToString(), /*bErrorIfNotFound=*/ false);
+
+	if (NewTag != TagSearch)
+	{
+		TagSearch = NewTag;
+		BroadcastChangedEvent();
+	}
+}
+
+void FFrontendFilter_GameplayTags::OnTagWidgetChanged()
+{
+	BroadcastChangedEvent();
+}
+
+bool FFrontendFilter_GameplayTags::ProcessStruct(void* Data, UStruct* Struct) const
+{
+	for (TFieldIterator<UProperty> PropIt(Struct, EFieldIteratorFlags::IncludeSuper); PropIt; ++PropIt)
+	{
+		UProperty* Prop = *PropIt;
+
+		if (ProcessProperty(Data, Prop))
+		{
+			return true;
+		}
+	}
+
+	return false;
+}
+
+bool FFrontendFilter_GameplayTags::ProcessProperty(void* Data, UProperty* Prop) const
+{
+	void* InnerData = Prop->ContainerPtrToValuePtr<void>(Data);
+
+	if (UStructProperty* StructProperty = Cast<UStructProperty>(Prop))
+	{
+		if (StructProperty->Struct == FGameplayTag::StaticStruct())
+		{
+			FGameplayTag& ThisTag = *static_cast<FGameplayTag*>(InnerData);
+
+			const bool bPassesTagSearch = (TagSearch.IsValid() == false) || TagSearch.Matches(EGameplayTagMatchType::Explicit, ThisTag, EGameplayTagMatchType::IncludeParentTags);
+
+			return bPassesTagSearch;
+		}
+		else
+		{
+			return ProcessStruct(InnerData, StructProperty->Struct);
+		}
+	}
+	else if (UArrayProperty* ArrayProperty = Cast<UArrayProperty>(Prop))
+	{
+		FScriptArrayHelper ArrayHelper(ArrayProperty, InnerData);
+		for (int32 ArrayIndex = 0; ArrayIndex < ArrayHelper.Num(); ++ArrayIndex)
+		{
+			void* ArrayData = ArrayHelper.GetRawPtr(ArrayIndex);
+
+			if (ProcessProperty(ArrayData, ArrayProperty->Inner))
+			{
+				return true;
+			}
+		}
+	}
+
+	return false;
+}
+
+bool FFrontendFilter_GameplayTags::PassesFilter(FAssetFilterType InItem) const
+{
+	if (InItem.IsAssetLoaded())
+	{
+		if (UObject* Object = InItem.GetAsset())
+		{
+			if (UBlueprint* Blueprint = Cast<UBlueprint>(Object))
+			{
+				return ProcessStruct(Blueprint->GeneratedClass->GetDefaultObject(), Blueprint->GeneratedClass);
+
+				//@TODO: Check blueprint bytecode!
+			}
+			else if (UClass* Class = Cast<UClass>(Object))
+			{
+				return ProcessStruct(Class->GetDefaultObject(), Class);
+			}
+			else
+			{
+				return ProcessStruct(Object, Object->GetClass());
+			}
+		}
+	}
+
+	return false;
+}
+
+//////////////////////////////////////////////////////////////////////////
+//
+
+void UGameplayTagSearchFilter::AddFrontEndFilterExtensions(TSharedPtr<FFrontendFilterCategory> DefaultCategory, TArray< TSharedRef<class FFrontendFilter> >& InOutFilterList) const
+{
+	InOutFilterList.Add(MakeShareable(new FFrontendFilter_GameplayTags(DefaultCategory)));
+}
+
+#undef LOCTEXT_NAMESPACE
\ No newline at end of file
diff --git a/Engine/Source/Editor/GameplayTagsEditor/Private/GameplayTagSearchFilter.h b/Engine/Source/Editor/GameplayTagsEditor/Private/GameplayTagSearchFilter.h
new file mode 100644
--- /dev/null
+++ b/Engine/Source/Editor/GameplayTagsEditor/Private/GameplayTagSearchFilter.h
@@ -0,0 +1,17 @@
+// Copyright 1998-2015 Epic Games, Inc. All Rights Reserved.
+
+#pragma once
+
+#include "ContentBrowserFrontEndFilterExtension.h"
+#include "GameplayTagSearchFilter.generated.h"
+
+UCLASS()
+class UGameplayTagSearchFilter : public UContentBrowserFrontEndFilterExtension
+{
+public:
+	GENERATED_BODY()
+
+	// UContentBrowserFrontEndFilterExtension interface
+	virtual void AddFrontEndFilterExtensions(TSharedPtr<class FFrontendFilterCategory> DefaultCategory, TArray< TSharedRef<class FFrontendFilter> >& InOutFilterList) const override;
+	// End of UContentBrowserFrontEndFilterExtension interface
+};
