Gameplay Tags: Improved CB filter to allow picking tags using the hierarchical tag list widget and allow multiple tags to be selected

diff --git a/Engine/Source/Editor/GameplayTagsEditor/Private/GameplayTagSearchFilter.cpp b/Engine/Source/Editor/GameplayTagsEditor/Private/GameplayTagSearchFilter.cpp
--- a/Engine/Source/Editor/GameplayTagsEditor/Private/GameplayTagSearchFilter.cpp
+++ b/Engine/Source/Editor/GameplayTagsEditor/Private/GameplayTagSearchFilter.cpp
@@ -20,12 +20,19 @@
 class FFrontendFilter_GameplayTags : public FFrontendFilter
 {
 public:
-	FFrontendFilter_GameplayTags(TSharedPtr<FFrontendFilterCategory> InCategory) : FFrontendFilter(InCategory) {}
+	FFrontendFilter_GameplayTags(TSharedPtr<FFrontendFilterCategory> InCategory)
+		: FFrontendFilter(InCategory)
+	{
+		TagContainer = MakeShareable(new FGameplayTagContainer);
+
+		EditableContainers.Add(SGameplayTagWidget::FEditableGameplayTagContainerDatum(/*TagContainerOwner=*/ nullptr, TagContainer.Get()));
+	}
 
 	// FFrontendFilter implementation
+	virtual FLinearColor GetColor() const override { return FLinearColor::Red; }
 	virtual FString GetName() const override { return TEXT("GameplayTagFilter"); }
 	virtual FText GetDisplayName() const override;
-	virtual FText GetToolTipText() const override { return LOCTEXT("GameplayTagFilterDisplayTooltip", "Search for any Blueprint or asset containing a specified gameplay tag (right-click to change the tag)."); }
+	virtual FText GetToolTipText() const override;
 	virtual void ModifyContextMenu(FMenuBuilder& MenuBuilder) override;
 	virtual void SaveSettings(const FString& IniFilename, const FString& IniSection, const FString& SettingsString) const override;
 	virtual void LoadSettings(const FString& IniFilename, const FString& IniSection, const FString& SettingsString) override;
@@ -36,16 +43,17 @@ public:
 	// End of IFilter implementation
 
 protected:
-	// If blank, finds anything that contains a gameplay tag
-	FGameplayTag TagSearch;
+	// Container of selected search tags (the asset is shown if *any* of these match)
+	TSharedPtr<FGameplayTagContainer> TagContainer;
+
+	// Adaptor for the SGameplayTagWidget to edit our tag container
+	TArray<SGameplayTagWidget::FEditableGameplayTagContainerDatum> EditableContainers;
 
 protected:
 	bool ProcessStruct(void* Data, UStruct* Struct) const;
 
 	bool ProcessProperty(void* Data, UProperty* Prop) const;
 
-	FText GetTagValueAsText() const;
-	void OnTagTextCommitted(const FText& InText, ETextCommit::Type InCommitType);
 	void OnTagWidgetChanged();
 };
 
@@ -53,60 +61,87 @@ void FFrontendFilter_GameplayTags::ModifyContextMenu(FMenuBuilder& MenuBuilder)
 {
 	FUIAction Action;
 
-	MenuBuilder.BeginSection(TEXT("ComparsionSection"), LOCTEXT("ComparisonSectionHeading", "Gameplay Tag Comparison"));
-
-	TSharedRef<SWidget> KeyWidget =
-		SNew(SEditableTextBox)
-		.Text_Raw(this, &FFrontendFilter_GameplayTags::GetTagValueAsText)
-		.OnTextCommitted_Raw(this, &FFrontendFilter_GameplayTags::OnTagTextCommitted)
-		.MinDesiredWidth(100.0f);
-
-	MenuBuilder.AddWidget(KeyWidget, LOCTEXT("TagMenuDesc", "Gameplay Tag"));
-
-// 	TSharedRef<SWidget> TagWidget = SNew(SGameplayTagWidget)
-// 		.MultiSelect(true)
-// 		.OnTagChanged(this, &FFrontendFilter_GameplayTags::OnTagWidgetChanged);
-// 	MenuBuilder.AddWidget(TagWidget, LOCTEXT("TagMenuDesc2", "Gameplay Tag"));
+	MenuBuilder.BeginSection(TEXT("ComparsionSection"), LOCTEXT("ComparisonSectionHeading", "Gameplay Tag(s) to search for"));
+
+	TSharedRef<SWidget> TagWidget =
+		SNew(SVerticalBox)
+		+SVerticalBox::Slot()
+		.AutoHeight()
+		.MaxHeight(300)
+		[
+			SNew(SGameplayTagWidget, EditableContainers)
+			.MultiSelect(true)
+			.OnTagChanged_Raw(this, &FFrontendFilter_GameplayTags::OnTagWidgetChanged)
+		];
+ 	MenuBuilder.AddWidget(TagWidget, FText::GetEmpty(), /*bNoIndent=*/ false);
 }
 
 FText FFrontendFilter_GameplayTags::GetDisplayName() const
 {
-	return FText::Format(LOCTEXT("FFrontendFilter_CompareOperation", "Gameplay Tags ({0})"),
-		TagSearch.IsValid() ? FText::FromName(TagSearch.GetTagName()) : LOCTEXT("AnyTag", "Any Tags"));
-}
+	if (TagContainer->Num() == 0)
+	{
+		return LOCTEXT("AnyGameplayTagDisplayName", "Gameplay Tags");
+	}
+	else
+	{
+		FString QueryString;
 
-void FFrontendFilter_GameplayTags::SaveSettings(const FString& IniFilename, const FString& IniSection, const FString& SettingsString) const
-{
-	GConfig->SetString(*IniSection, *(SettingsString + TEXT(".Key")), *TagSearch.ToString(), IniFilename);
+		int32 Count = 0;
+		for (const FGameplayTag& Tag : *TagContainer.Get())
+		{
+			if (Count > 0)
+			{
+				QueryString += TEXT(" | ");
+			}
+
+			QueryString += Tag.ToString();
+			++Count;
+		}
+
+
+		return FText::Format(LOCTEXT("GameplayTagListDisplayName", "Gameplay Tags ({0})"), FText::AsCultureInvariant(QueryString));
+	}
 }
 
-void FFrontendFilter_GameplayTags::LoadSettings(const FString& IniFilename, const FString& IniSection, const FString& SettingsString)
+FText FFrontendFilter_GameplayTags::GetToolTipText() const
 {
-	FString TagNameAsString;
-	if (GConfig->GetString(*IniSection, *(SettingsString + TEXT(".Key")), TagNameAsString, IniFilename))
+	if (TagContainer->Num() == 0)
 	{
-		IGameplayTagsModule& GameplayTagModule = IGameplayTagsModule::Get();
-
-		FGameplayTag NewTag = GameplayTagModule.RequestGameplayTag(*TagNameAsString, /*bErrorIfNotFound=*/ false);
-		TagSearch = NewTag;
+		return LOCTEXT("AnyGameplayTagFilterDisplayTooltip", "Search for any *loaded* Blueprint or asset that contains a gameplay tag (right-click to choose tags).");
+	}
+	else
+	{
+		return LOCTEXT("GameplayTagFilterDisplayTooltip", "Search for any *loaded* Blueprint or asset that has a gameplay tag which matches any of the selected tags (right-click to choose tags).");
 	}
 }
 
-FText FFrontendFilter_GameplayTags::GetTagValueAsText() const
+void FFrontendFilter_GameplayTags::SaveSettings(const FString& IniFilename, const FString& IniSection, const FString& SettingsString) const
 {
-	return FText::FromName(TagSearch.GetTagName());
+	TArray<FString> TagStrings;
+	TagStrings.Reserve(TagContainer->Num());
+	for (const FGameplayTag& Tag : *TagContainer.Get())
+	{
+		TagStrings.Add(Tag.GetTagName().ToString());
+	}
+
+	GConfig->SetArray(*IniSection, *(SettingsString + TEXT(".Tags")), TagStrings, IniFilename);
 }
 
-void FFrontendFilter_GameplayTags::OnTagTextCommitted(const FText& InText, ETextCommit::Type InCommitType)
+void FFrontendFilter_GameplayTags::LoadSettings(const FString& IniFilename, const FString& IniSection, const FString& SettingsString)
 {
 	IGameplayTagsModule& GameplayTagModule = IGameplayTagsModule::Get();
 
-	FGameplayTag NewTag = GameplayTagModule.RequestGameplayTag(*InText.ToString(), /*bErrorIfNotFound=*/ false);
+	TArray<FString> TagStrings;
+	GConfig->GetArray(*IniSection, *(SettingsString + TEXT(".Tags")), /*out*/ TagStrings, IniFilename);
 
-	if (NewTag != TagSearch)
+	TagContainer->RemoveAllTags();
+	for (const FString& TagString : TagStrings)
 	{
-		TagSearch = NewTag;
-		BroadcastChangedEvent();
+		FGameplayTag NewTag = GameplayTagModule.RequestGameplayTag(*TagString, /*bErrorIfNotFound=*/ false);
+		if (NewTag.IsValid())
+		{
+			TagContainer->AddTag(NewTag);
+		}
 	}
 }
 
@@ -140,7 +175,8 @@ bool FFrontendFilter_GameplayTags::ProcessProperty(void* Data, UProperty* Prop)
 		{
 			FGameplayTag& ThisTag = *static_cast<FGameplayTag*>(InnerData);
 
-			const bool bPassesTagSearch = (TagSearch.IsValid() == false) || TagSearch.Matches(EGameplayTagMatchType::Explicit, ThisTag, EGameplayTagMatchType::IncludeParentTags);
+			const bool bAnyTagIsOK = TagContainer->Num() == 0;
+			const bool bPassesTagSearch = bAnyTagIsOK || TagContainer->HasTag(ThisTag, EGameplayTagMatchType::Explicit, EGameplayTagMatchType::IncludeParentTags);
 
 			return bPassesTagSearch;
 		}
