Fixed some indentation issues in FGameplayTagContainerCustomization

diff --git a/Engine/Source/Editor/GameplayTagsEditor/Private/GameplayTagContainerCustomization.cpp b/Engine/Source/Editor/GameplayTagsEditor/Private/GameplayTagContainerCustomization.cpp
--- a/Engine/Source/Editor/GameplayTagsEditor/Private/GameplayTagContainerCustomization.cpp
+++ b/Engine/Source/Editor/GameplayTagsEditor/Private/GameplayTagContainerCustomization.cpp
@@ -28,20 +28,20 @@ void FGameplayTagContainerCustomization::CustomizeHeader(TSharedRef<class IPrope
 		.MaxDesiredWidth(512)
 		[
 			SNew(SHorizontalBox)
-		+ SHorizontalBox::Slot()
-		.AutoWidth()
-		.VAlign(VAlign_Center)
-		[
-			SNew(SVerticalBox)
-			+ SVerticalBox::Slot()
-			.AutoHeight()
+			+ SHorizontalBox::Slot()
+			.AutoWidth()
+			.VAlign(VAlign_Center)
 			[
-				SNew(SButton)
-				.IsEnabled(!StructPropertyHandle->GetProperty()->HasAnyPropertyFlags(CPF_EditConst))
-				.Text(LOCTEXT("GameplayTagContainerCustomization_Edit", "Edit..."))
-				.OnClicked(this, &FGameplayTagContainerCustomization::OnEditButtonClicked)
-			]
-			+ SVerticalBox::Slot()
+				SNew(SVerticalBox)
+				+ SVerticalBox::Slot()
+				.AutoHeight()
+				[
+					SNew(SButton)
+					.IsEnabled(!StructPropertyHandle->GetProperty()->HasAnyPropertyFlags(CPF_EditConst))
+					.Text(LOCTEXT("GameplayTagContainerCustomization_Edit", "Edit..."))
+					.OnClicked(this, &FGameplayTagContainerCustomization::OnEditButtonClicked)
+				]
+				+ SVerticalBox::Slot()
 				.AutoHeight()
 				[
 					SNew(SButton)
@@ -50,16 +50,16 @@ void FGameplayTagContainerCustomization::CustomizeHeader(TSharedRef<class IPrope
 					.OnClicked(this, &FGameplayTagContainerCustomization::OnClearAllButtonClicked)
 					.Visibility(this, &FGameplayTagContainerCustomization::GetClearAllVisibility)
 				]
-		]
-	+ SHorizontalBox::Slot()
-		.AutoWidth()
-		[
-			SNew(SBorder)
-			.Padding(4.0f)
-			.Visibility(this, &FGameplayTagContainerCustomization::GetTagsListVisibility)
-			[
-				ActiveTags()
 			]
+			+ SHorizontalBox::Slot()
+			.AutoWidth()
+			[
+				SNew(SBorder)
+				.Padding(4.0f)
+				.Visibility(this, &FGameplayTagContainerCustomization::GetTagsListVisibility)
+				[
+					ActiveTags()
+				]
 			]
 		];
 
