Gameplay Tags: Removed #pragma once from a .cpp

diff --git a/Engine/Source/Editor/GameplayTagsEditor/Private/GameplayTagSearchFilter.cpp b/Engine/Source/Editor/GameplayTagsEditor/Private/GameplayTagSearchFilter.cpp
--- a/Engine/Source/Editor/GameplayTagsEditor/Private/GameplayTagSearchFilter.cpp
+++ b/Engine/Source/Editor/GameplayTagsEditor/Private/GameplayTagSearchFilter.cpp
@@ -1,7 +1,5 @@
 // Copyright 1998-2015 Epic Games, Inc. All Rights Reserved.
 
-#pragma once
-
 #include "GameplayTagsEditorModulePrivatePCH.h"
 #include "GameplayTagSearchFilter.h"
 
