Fixed Tag and Tag Query Editor windows on Mac
[UE-18455] - Tag Editor window cannot be focused on Mac

#fix FGameplayTagQueryCustomization and FGameplayTagContainerCustomization checks FGlobalTabmanager::GetRootWindow() before trying to create a child window. This stops child windows on Mac.

diff --git a/Engine/Source/Editor/GameplayTagsEditor/Private/GameplayTagContainerCustomization.cpp b/Engine/Source/Editor/GameplayTagsEditor/Private/GameplayTagContainerCustomization.cpp
--- a/Engine/Source/Editor/GameplayTagsEditor/Private/GameplayTagContainerCustomization.cpp
+++ b/Engine/Source/Editor/GameplayTagsEditor/Private/GameplayTagContainerCustomization.cpp
@@ -156,10 +156,9 @@ FReply FGameplayTagContainerCustomization::OnEditButtonClicked()
 
 	GameplayTagWidgetWindow->GetOnWindowDeactivatedEvent().AddRaw(this, &FGameplayTagContainerCustomization::OnGameplayTagWidgetWindowDeactivate);
 
-	IMainFrameModule& MainFrameModule = FModuleManager::LoadModuleChecked<IMainFrameModule>(TEXT("MainFrame"));
-	if (MainFrameModule.GetParentWindow().IsValid())
+	if (FGlobalTabmanager::Get()->GetRootWindow().IsValid())
 	{
-		FSlateApplication::Get().AddWindowAsNativeChild(GameplayTagWidgetWindow.ToSharedRef(), MainFrameModule.GetParentWindow().ToSharedRef());
+		FSlateApplication::Get().AddWindowAsNativeChild(GameplayTagWidgetWindow.ToSharedRef(), FGlobalTabmanager::Get()->GetRootWindow().ToSharedRef());
 	}
 	else
 	{
diff --git a/Engine/Source/Editor/GameplayTagsEditor/Private/GameplayTagQueryCustomization.cpp b/Engine/Source/Editor/GameplayTagsEditor/Private/GameplayTagQueryCustomization.cpp
--- a/Engine/Source/Editor/GameplayTagsEditor/Private/GameplayTagQueryCustomization.cpp
+++ b/Engine/Source/Editor/GameplayTagsEditor/Private/GameplayTagQueryCustomization.cpp
@@ -181,10 +181,9 @@ FReply FGameplayTagQueryCustomization::OnEditButtonClicked()
 				.ReadOnly(bReadOnly)
 			];
 
-		IMainFrameModule& MainFrameModule = FModuleManager::LoadModuleChecked<IMainFrameModule>(TEXT("MainFrame"));
-		if (MainFrameModule.GetParentWindow().IsValid())
+		if (FGlobalTabmanager::Get()->GetRootWindow().IsValid())
 		{
-			FSlateApplication::Get().AddWindowAsNativeChild(GameplayTagQueryWidgetWindow.ToSharedRef(), MainFrameModule.GetParentWindow().ToSharedRef());
+			FSlateApplication::Get().AddWindowAsNativeChild(GameplayTagQueryWidgetWindow.ToSharedRef(), FGlobalTabmanager::Get()->GetRootWindow().ToSharedRef());
 		}
 		else
 		{
