Fixed SGameplayTagWidget for large projects with many tags by removing SScalebox
[UE-18455] - Tag Editor window cannot be focused on Mac

diff --git a/Engine/Source/Editor/GameplayTagsEditor/Private/GameplayTagCustomization.cpp b/Engine/Source/Editor/GameplayTagsEditor/Private/GameplayTagCustomization.cpp
--- a/Engine/Source/Editor/GameplayTagsEditor/Private/GameplayTagCustomization.cpp
+++ b/Engine/Source/Editor/GameplayTagsEditor/Private/GameplayTagCustomization.cpp
@@ -75,21 +75,13 @@ TSharedRef<SWidget> FGameplayTagCustomization::GetListContent()
 		.AutoHeight()
 		.MaxHeight(400)
 		[
-			SNew(SScaleBox)
-			.HAlign(EHorizontalAlignment::HAlign_Left)
-			.VAlign(EVerticalAlignment::VAlign_Top)
-			.StretchDirection(EStretchDirection::DownOnly)
-			.Stretch(EStretch::ScaleToFit)
-			.Content()
-			[
-				SNew(SGameplayTagWidget, EditableContainers)
-				.Filter(Categories)
-				.ReadOnly(bReadOnly)
-				.TagContainerName(StructPropertyHandle->GetPropertyDisplayName().ToString())
-				.MultiSelect(false)
-				.OnTagChanged(this, &FGameplayTagCustomization::OnTagChanged)
-				.PropertyHandle(StructPropertyHandle)
-			]
+            SNew(SGameplayTagWidget, EditableContainers)
+            .Filter(Categories)
+            .ReadOnly(bReadOnly)
+            .TagContainerName(StructPropertyHandle->GetPropertyDisplayName().ToString())
+            .MultiSelect(false)
+            .OnTagChanged(this, &FGameplayTagCustomization::OnTagChanged)
+            .PropertyHandle(StructPropertyHandle)
 		];
 }
 
