Remove SScalebox since it causes scaling issues with large amounts of gameplay tags

--------
Integrated using branch UE4-To-UE4-Orion (reversed) of change#2677822 by [REDACTED_PERSON] on 2015/09/02 18:29:09.

diff --git a/Engine/Source/Editor/GameplayTagsEditor/Private/SGameplayTagContainerGraphPin.cpp b/Engine/Source/Editor/GameplayTagsEditor/Private/SGameplayTagContainerGraphPin.cpp
--- a/Engine/Source/Editor/GameplayTagsEditor/Private/SGameplayTagContainerGraphPin.cpp
+++ b/Engine/Source/Editor/GameplayTagsEditor/Private/SGameplayTagContainerGraphPin.cpp
@@ -105,18 +105,10 @@ TSharedRef<SWidget> SGameplayTagContainerGraphPin::GetListContent()
 		.AutoHeight()
 		.MaxHeight( 400 )
 		[
-			SNew(SScaleBox)
-			.HAlign(EHorizontalAlignment::HAlign_Left)
-			.VAlign(EVerticalAlignment::VAlign_Top)
-			.StretchDirection(EStretchDirection::DownOnly)
-			.Stretch(EStretch::ScaleToFit)
-			.Content()
-			[
-				SNew( SGameplayTagWidget, EditableContainers )
-				.OnTagChanged(this, &SGameplayTagContainerGraphPin::RefreshTagList)
-				.TagContainerName( TEXT("SGameplayTagContainerGraphPin") )
-				.Visibility( this, &SGraphPin::GetDefaultValueVisibility )
-			]
+			SNew( SGameplayTagWidget, EditableContainers )
+			.OnTagChanged(this, &SGameplayTagContainerGraphPin::RefreshTagList)
+			.TagContainerName( TEXT("SGameplayTagContainerGraphPin") )
+			.Visibility( this, &SGraphPin::GetDefaultValueVisibility )
 		];
 }
 
